Fix template checking in SAssetTagItem.h

When running in MSVC with the /permissive- flag this file was failing to compile due to the (correct C++) early checking of non dependent names.  The fix fix was to convert the failing names to dependent lookups so that they can be resolved at instantiation time.

diff --git a/Engine/Source/Editor/AssetTagsEditor/Public/SAssetTagItem.h b/Engine/Source/Editor/AssetTagsEditor/Public/SAssetTagItem.h
--- a/Engine/Source/Editor/AssetTagsEditor/Public/SAssetTagItem.h
+++ b/Engine/Source/Editor/AssetTagsEditor/Public/SAssetTagItem.h
@@ -16,14 +16,14 @@ template <typename ItemType>
 class SAssetTagItemTableRow : public STableRow<ItemType>
 {
 public:
-	void Construct(const FArguments& InArgs, const TSharedRef<STableViewBase>& InOwnerTableView)
+	void Construct(const typename STableRow<ItemType>::FArguments& InArgs, const TSharedRef<STableViewBase>& InOwnerTableView)
 	{
-		ConstructInternal(InArgs, InOwnerTableView);
+		this->ConstructInternal(InArgs, InOwnerTableView);
 
-		Style = &FEditorStyle::GetWidgetStyle<FTableRowStyle>("ContentBrowser.AssetTagTableRow");
-		BorderImage = FEditorStyle::GetBrush("NoBorder");
+		STableRow<ItemType>::Style = &FEditorStyle::GetWidgetStyle<FTableRowStyle>("ContentBrowser.AssetTagTableRow");
+		STableRow<ItemType>::BorderImage = FEditorStyle::GetBrush("NoBorder");
 
-		ConstructChildren(
+		this->ConstructChildren(
 			InOwnerTableView->TableViewMode,
 			ApplySelectionBorderOuterPadding(InArgs._Padding),
 			WrapContentInSelectionBorder(InArgs._Content.Widget)
@@ -37,17 +37,17 @@ public:
 
 	virtual void SetContent(TSharedRef<SWidget> InContent) override
 	{
-		STableRow::SetContent(WrapContentInSelectionBorder(InContent));
+		this->SetContent(WrapContentInSelectionBorder(InContent));
 	}
 
 	virtual const FSlateBrush* GetBorder() const override
 	{
 		if (IsDropTarget.Get(false))
 		{
-			return &Style->InactiveHoveredBrush;
+			return &STableRow<ItemType>::Style->InactiveHoveredBrush;
 		}
 
-		return STableRow::GetBorder();
+		return this->GetBorder();
 	}
 
 private:
