Fix two infinite recursions inadvertently added in the last fix

diff --git a/Engine/Source/Editor/AssetTagsEditor/Public/SAssetTagItem.h b/Engine/Source/Editor/AssetTagsEditor/Public/SAssetTagItem.h
--- a/Engine/Source/Editor/AssetTagsEditor/Public/SAssetTagItem.h
+++ b/Engine/Source/Editor/AssetTagsEditor/Public/SAssetTagItem.h
@@ -20,8 +20,8 @@ public:
 	{
 		this->ConstructInternal(InArgs, InOwnerTableView);
 
-		STableRow<ItemType>::Style = &FEditorStyle::GetWidgetStyle<FTableRowStyle>("ContentBrowser.AssetTagTableRow");
-		STableRow<ItemType>::BorderImage = FEditorStyle::GetBrush("NoBorder");
+		this->Style = &FEditorStyle::GetWidgetStyle<FTableRowStyle>("ContentBrowser.AssetTagTableRow");
+		this->BorderImage = FEditorStyle::GetBrush("NoBorder");
 
 		this->ConstructChildren(
 			InOwnerTableView->TableViewMode,
@@ -37,17 +37,17 @@ public:
 
 	virtual void SetContent(TSharedRef<SWidget> InContent) override
 	{
-		this->SetContent(WrapContentInSelectionBorder(InContent));
+		STableRow<ItemType>::SetContent(WrapContentInSelectionBorder(InContent));
 	}
 
 	virtual const FSlateBrush* GetBorder() const override
 	{
 		if (IsDropTarget.Get(false))
 		{
-			return &STableRow<ItemType>::Style->InactiveHoveredBrush;
+			return &this->Style->InactiveHoveredBrush;
 		}
 
-		return this->GetBorder();
+		return STableRow<ItemType>::GetBorder();
 	}
 
 private:
