Fix Linux compile error with SAssetTagItem.cpp
#rnx

diff --git a/Engine/Source/Editor/AssetTagsEditor/Private/SAssetTagItem.cpp b/Engine/Source/Editor/AssetTagsEditor/Private/SAssetTagItem.cpp
--- a/Engine/Source/Editor/AssetTagsEditor/Private/SAssetTagItem.cpp
+++ b/Engine/Source/Editor/AssetTagsEditor/Private/SAssetTagItem.cpp
@@ -145,7 +145,7 @@ void SAssetTagItem::Construct(const FArguments& InArgs)
 		.DisplayName(InArgs._DisplayName)
 		.OnBuildToolTipInfo(InArgs._OnBuildToolTipInfo);
 
-	ANSICHAR* StyleSpecifier = nullptr;
+	const ANSICHAR* StyleSpecifier = nullptr;
 	if (InArgs._ViewMode == EAssetTagItemViewMode::Compact)
 	{
 		StyleSpecifier = ".Compact";
